[Substrate] Set Blendable GBuffer format as default for new projects. Some archvis and automotive project templates are set to Adapative GBuffer format by default. Also fixed the popup window when enabling Substrate saying that the feature is still in beta.

diff --git a/Engine/Source/Editor/GameProjectGeneration/Private/GameProjectUtils.cpp b/Engine/Source/Editor/GameProjectGeneration/Private/GameProjectUtils.cpp
--- a/Engine/Source/Editor/GameProjectGeneration/Private/GameProjectUtils.cpp
+++ b/Engine/Source/Editor/GameProjectGeneration/Private/GameProjectUtils.cpp
@@ -203,6 +203,12 @@ namespace
 			TEXT("r.Substrate"),
 			TEXT("True"),
 			false /* ShouldReplaceExistingValue */);
+
+		ConfigValues.Emplace(TEXT("DefaultEngine.ini"),
+			TEXT("/Script/Engine.RendererSettings"),
+			TEXT("r.Substrate.ProjectGBufferFormat"),
+			TEXT("0"),
+			false /* ShouldReplaceExistingValue */);
 	}
 
 	/** Get the configuration values for enabling WorldPartition by default. */
diff --git a/Engine/Source/Runtime/Engine/Private/RendererSettings.cpp b/Engine/Source/Runtime/Engine/Private/RendererSettings.cpp
--- a/Engine/Source/Runtime/Engine/Private/RendererSettings.cpp
+++ b/Engine/Source/Runtime/Engine/Private/RendererSettings.cpp
@@ -364,7 +364,7 @@ void URendererSettings::PostEditChangeProperty(FPropertyChangedEvent& PropertyCh
 		{
 			if (bEnableSubstrate)
 			{
-				const EAppReturnType::Type Choice = FMessageDialog::Open(EAppMsgType::OkCancel, LOCTEXT("Substrate Beta", "Warning: Substrate is in a Beta state. Be aware that any materials saved when Substrate is enabled won't be rendered correctly if Substrate is disabled later on."));
+				const EAppReturnType::Type Choice = FMessageDialog::Open(EAppMsgType::OkCancel, LOCTEXT("Substrate", "Warning: Be aware that any materials saved when Substrate is enabled won't be rendered correctly if Substrate is disabled later on."));
 				if (Choice == EAppReturnType::Cancel)
 				{
 					bEnableSubstrate = false;
diff --git a/Templates/TP_AEC_ArchvisBP/Config/DefaultEngine.ini b/Templates/TP_AEC_ArchvisBP/Config/DefaultEngine.ini
--- a/Templates/TP_AEC_ArchvisBP/Config/DefaultEngine.ini
+++ b/Templates/TP_AEC_ArchvisBP/Config/DefaultEngine.ini
@@ -39,6 +39,8 @@ r.RayTracing.UseTextureLod=True
 r.Lumen.HardwareRayTracing=True
 r.Lumen.TraceMeshSDF=1
 r.RayTracing.Geometry.InstancedStaticMeshes.Culling=0
+r.Substrate=True
+r.Substrate.ProjectGBufferFormat=1
 
 r.MinRoughnessOverride=0.02
 
diff --git a/Templates/TP_AEC_BlankBP/Config/DefaultEngine.ini b/Templates/TP_AEC_BlankBP/Config/DefaultEngine.ini
--- a/Templates/TP_AEC_BlankBP/Config/DefaultEngine.ini
+++ b/Templates/TP_AEC_BlankBP/Config/DefaultEngine.ini
@@ -38,6 +38,8 @@ r.Shadow.Virtual.Enable=1
 r.RayTracing.UseTextureLod=True
 r.Lumen.HardwareRayTracing=True
 r.Lumen.TraceMeshSDF=1
+r.Substrate=True
+r.Substrate.ProjectGBufferFormat=1
 
 r.MinRoughnessOverride=0.02
 
diff --git a/Templates/TP_AEC_ProdConfigBP/Config/DefaultEngine.ini b/Templates/TP_AEC_ProdConfigBP/Config/DefaultEngine.ini
--- a/Templates/TP_AEC_ProdConfigBP/Config/DefaultEngine.ini
+++ b/Templates/TP_AEC_ProdConfigBP/Config/DefaultEngine.ini
@@ -14,6 +14,8 @@ AppliedDefaultGraphicsPerformance=Maximum
 r.DefaultFeature.AutoExposure.ExtendDefaultLuminanceRange=True
 r.SupportSkyAtmosphereAffectsHeightFog=True
 r.AllowStaticLighting=False
+r.Substrate=True
+r.Substrate.ProjectGBufferFormat=1
 
 [/Script/EngineSettings.GameMapsSettings]
 EditorStartupMap=/Game/ProductAssets/Levels/Main.Main
diff --git a/Templates/TP_MFG_BlankBP/Config/DefaultEngine.ini b/Templates/TP_MFG_BlankBP/Config/DefaultEngine.ini
--- a/Templates/TP_MFG_BlankBP/Config/DefaultEngine.ini
+++ b/Templates/TP_MFG_BlankBP/Config/DefaultEngine.ini
@@ -18,11 +18,8 @@ r.DefaultFeature.LocalExposure.HighlightContrastScale=0.8
 r.DefaultFeature.LocalExposure.ShadowContrastScale=0.8
 r.SkinCache.CompileShaders=True
 r.RayTracing=True
-r.Substrate=True
 r.VirtualTextures=True
 r.GBufferFormat=3
-r.Substrate.BytesPerPixel=240
-r.Substrate.EnableLayerSupport=True
 r.ReflectionCaptureResolution=1024
 r.RayTracing.Culling=0
 r.SkyAtmosphere.MultiScatteringLUT.HighQuality=1
@@ -39,8 +36,12 @@ r.Translucent.UsesRectLights=True
 r.RayTracedTranslucency.Allowed=1
 r.RayTracedTranslucency.UseRayTracedRefraction=1
 r.Lumen.TranslucencyReflections.FrontLayer.EnableForProject=True
-r.Substrate.Debug.AdvancedVisualizationShaders=True
 r.PathTracing.VisibleLights=1
+r.Substrate=True
+r.Substrate.ProjectGBufferFormat=1
+r.Substrate.BytesPerPixel=240
+r.Substrate.EnableLayerSupport=True
+r.Substrate.Debug.AdvancedVisualizationShaders=True
 
 [/Script/EngineSettings.GameMapsSettings]
 EditorStartupMap=/Game/Main.Main
diff --git a/Templates/TP_MFG_ProdConfigBP/Config/DefaultEngine.ini b/Templates/TP_MFG_ProdConfigBP/Config/DefaultEngine.ini
--- a/Templates/TP_MFG_ProdConfigBP/Config/DefaultEngine.ini
+++ b/Templates/TP_MFG_ProdConfigBP/Config/DefaultEngine.ini
@@ -17,6 +17,8 @@ r.DefaultFeature.LocalExposure.HighlightContrastScale=1.0
 r.DefaultFeature.LocalExposure.ShadowContrastScale=1.0
 r.GenerateMeshDistanceFields=True
 r.AllowStaticLighting=False
+r.Substrate=True
+r.Substrate.ProjectGBufferFormat=1
 
 [/Script/EngineSettings.GameMapsSettings]
 EditorStartupMap=/Game/ProductAssets/Levels/Main.Main
