[Template-MFG_Blank] Update DefaultEngine.ini, DefaultEditorPerProjectUserSettings.ini - InterchangeProjectSettings, and add Interchange .uasset files.

#virtualized

diff --git a/Templates/TP_MFG_BlankBP/Config/DefaultEditorPerProjectUserSettings.ini b/Templates/TP_MFG_BlankBP/Config/DefaultEditorPerProjectUserSettings.ini
--- a/Templates/TP_MFG_BlankBP/Config/DefaultEditorPerProjectUserSettings.ini
+++ b/Templates/TP_MFG_BlankBP/Config/DefaultEditorPerProjectUserSettings.ini
@@ -3,4 +3,7 @@
 ;Options=(Geometry=UseUnrealNurbsTessellation)
 
 ;Import Rhino meshes and UVs
-Options=(Geometry=UseRenderMeshes)
\ No newline at end of file
+Options=(Geometry=UseRenderMeshes)
+
+[InterchangeImportDialogOptions]
+ShowEssentials=True
diff --git a/Templates/TP_MFG_BlankBP/Config/DefaultEngine.ini b/Templates/TP_MFG_BlankBP/Config/DefaultEngine.ini
--- a/Templates/TP_MFG_BlankBP/Config/DefaultEngine.ini
+++ b/Templates/TP_MFG_BlankBP/Config/DefaultEngine.ini
@@ -56,6 +56,14 @@ r.Translucency.StandardSeparated=1
 r.Translucent.UsesRectLights=True
 r.VirtualTextures=True
 
+[/Script/InterchangeEngine.InterchangeProjectSettings]
+ContentImportSettings=(DefaultPipelineStackOverride=((Materials, "Materials"),(Textures, "Textures"),(Sounds, "Audio")),ShowImportDialogOverride=((Textures, (bShowImportDialog=False))),PipelineStacks=(("Static Meshes", (Pipelines=("/Game/InterchangePipeline/Pipeline_StaticMeshes.Pipeline_StaticMeshes"),PerTranslatorPipelines=((Translator="/Script/InterchangeImport.InterchangeMaterialXTranslator",Pipelines=("/Interchange/Pipelines/DefaultAssetsPipeline.DefaultAssetsPipeline","/Interchange/Pipelines/DefaultMaterialXPipeline.DefaultMaterialXPipeline")),(Translator="/Script/InterchangeImport.InterchangeGLTFTranslator",Pipelines=("/Interchange/Pipelines/DefaultGLTFAssetsPipeline.DefaultGLTFAssetsPipeline","/Interchange/Pipelines/DefaultGLTFPipeline.DefaultGLTFPipeline"))))),("Skeletal Meshes", (Pipelines=("/Game/InterchangePipeline/Pipeline_SkeletalMeshes.Pipeline_SkeletalMeshes"),PerTranslatorPipelines=((Translator="/Script/InterchangeImport.InterchangeMaterialXTranslator",Pipelines=("/Interchange/Pipelines/DefaultAssetsPipeline.DefaultAssetsPipeline","/Interchange/Pipelines/DefaultMaterialXPipeline.DefaultMaterialXPipeline")),(Translator="/Script/InterchangeImport.InterchangeGLTFTranslator",Pipelines=("/Interchange/Pipelines/DefaultGLTFAssetsPipeline.DefaultGLTFAssetsPipeline","/Interchange/Pipelines/DefaultGLTFPipeline.DefaultGLTFPipeline"))))),("Audio", (Pipelines=("/Interchange/Pipelines/DefaultAudioPipeline.DefaultAudioPipeline"))),("Materials & Textures", (Pipelines=("/Interchange/Pipelines/DefaultMaterialPipeline.DefaultMaterialPipeline"),PerTranslatorPipelines=((Translator="/Script/InterchangeImport.InterchangeMaterialXTranslator",Pipelines=("/Interchange/Pipelines/DefaultMaterialPipeline.DefaultMaterialPipeline","/Interchange/Pipelines/DefaultMaterialXPipeline.DefaultMaterialXPipeline"))))),("Textures", (Pipelines=("/Interchange/Pipelines/DefaultTexturePipeline.DefaultTexturePipeline"))),("HDR Cubemap", (Pipelines=("/Game/InterchangePipeline/Pipeline_HDR_Cubemap.Pipeline_HDR_Cubemap")))),DefaultPipelineStack="Static Meshes",ImportDialogClass="/Script/InterchangeEditorPipelines.InterchangePipelineConfigurationGeneric",bShowImportDialog=True,bShowReimportDialog=False)
+SceneImportSettings=(PerTranslatorDialogOverride=,PipelineStacks=(("Scene", (Pipelines=("/Interchange/Pipelines/DefaultSceneAssetsPipeline.DefaultSceneAssetsPipeline","/Interchange/Pipelines/DefaultSceneLevelPipeline.DefaultSceneLevelPipeline"),PerTranslatorPipelines=((Translator="/Script/InterchangeImport.InterchangeMaterialXTranslator",Pipelines=("/Interchange/Pipelines/DefaultSceneAssetsPipeline.DefaultSceneAssetsPipeline","/Interchange/Pipelines/DefaultSceneLevelPipeline.DefaultSceneLevelPipeline","/Interchange/Pipelines/DefaultMaterialXPipeline.DefaultMaterialXPipeline")),(Translator="/Script/InterchangeImport.InterchangeGLTFTranslator",Pipelines=("/Interchange/Pipelines/DefaultGLTFSceneAssetsPipeline.DefaultGLTFSceneAssetsPipeline","/Interchange/Pipelines/DefaultSceneLevelPipeline.DefaultSceneLevelPipeline","/Interchange/Pipelines/DefaultGLTFPipeline.DefaultGLTFPipeline")))))),DefaultPipelineStack="Scene",ImportDialogClass="/Script/InterchangeEditorPipelines.InterchangePipelineConfigurationGeneric",bShowImportDialog=True,bShowReimportDialog=False)
+FilePickerClass=/Script/InterchangeEditorUtilities.InterchangeFilePickerGeneric
+bStaticMeshUseSmoothEdgesIfSmoothingInformationIsMissing=True
+GenericPipelineClass=/Script/InterchangePipelines.InterchangeGenericAssetsPipeline
+ConverterDefaultPipeline=/Interchange/Pipelines/DefaultAssetsPipeline.DefaultAssetsPipeline
+
 [/Script/EngineSettings.GameMapsSettings]
 EditorStartupMap=/Game/Main.Main
 +EditorTemplateMapOverrides=(Thumbnail="/Game/BlankDefault/Levels/Thumbnails/Thumbnail_Empty.Thumbnail_Empty",Map="/Game/BlankDefault/Levels/Blank_Empty.Blank_Empty",DisplayName=NSLOCTEXT("[/Script/EngineSettings]", "E1C1BC224C1F413942879D8305132FF6", "Default Level"))
