[Template-MFG_Blank] TemplateDefs.ini - Uncomment SharedContent BlankDefault to avoid unnecessary asset loading and use the MFG-specific blank level.

diff --git a/Templates/TP_MFG_BlankBP/Config/TemplateDefs.ini b/Templates/TP_MFG_BlankBP/Config/TemplateDefs.ini
--- a/Templates/TP_MFG_BlankBP/Config/TemplateDefs.ini
+++ b/Templates/TP_MFG_BlankBP/Config/TemplateDefs.ini
@@ -108,4 +108,4 @@ HiddenSettings=All
 
 ; Shared feature packs. The files in these packs listed in these structures marked as 'additionalfiles' will be copied on project generation.
 ;SharedContentPacks=(MountName="2DSideScroller",DetailLevels=("Standard"))
-SharedContentPacks=(MountName="BlankDefault",DetailLevels=("Standard"))
\ No newline at end of file
+;SharedContentPacks=(MountName="BlankDefault",DetailLevels=("Standard"))
\ No newline at end of file
