[Template-PhotoStudio / Template-MFG_Blank] Update rendering settings and Post Process Volume settings to align with Automotive Substrate Materials.

#virtualized
#Synced-CL 51391350

#Synced-CL 51391350

diff --git a/Templates/TP_MFG_BlankBP/Config/DefaultEngine.ini b/Templates/TP_MFG_BlankBP/Config/DefaultEngine.ini
--- a/Templates/TP_MFG_BlankBP/Config/DefaultEngine.ini
+++ b/Templates/TP_MFG_BlankBP/Config/DefaultEngine.ini
@@ -7,7 +7,7 @@ AppliedDefaultGraphicsPerformance=Maximum
 [/Script/Engine.RendererSettings]
 r.AllowStaticLighting=0
 r.AntiAliasingMethod=4
-r.ClearCoatNormal=True
+r.BloomQuality=2
 r.DefaultFeature.AutoExposure.ExtendDefaultLuminanceRange=True
 r.DefaultFeature.AutoExposure=False
 r.DefaultFeature.LocalExposure.HighlightContrastScale=0.8
@@ -19,43 +19,56 @@ r.GenerateMeshDistanceFields=True
 r.Lumen.HardwareRayTracing.HitLighting.ForceOpaque=0
 r.Lumen.HardwareRayTracing.LightingMode=1
 r.Lumen.LightingDataFormat=1
+r.Lumen.Reflections.BilateralFilter.KernelRadius=32
 r.Lumen.Reflections.DenoiserTonemapRange =12
 r.Lumen.Reflections.HardwareRayTracing.Translucent.MaxRefractionBounces=2
+r.Lumen.ScreenProbeGather.DownsampleFactor=4
+r.Lumen.ScreenProbeGather.ShortRangeAO.HorizonSearch.ForegroundSampleRejectDistanceFraction=0.025
+r.Lumen.ScreenProbeGather.Temporal.DistanceThreshold=0.001
+r.Lumen.ScreenProbeGather.Temporal.NormalThreshold=5
+r.Lumen.ScreenProbeGather.Temporal.RejectBasedOnNormal=1
+r.Lumen.ScreenProbeGather.TemporalFilterProbes=1
 r.Lumen.TraceMeshSDFs=0
 r.Lumen.TranslucencyReflections.FrontLayer.EnableForProject=True
+r.LumenScene.SurfaceCache.CardMinResolution=2
+r.LumenScene.SurfaceCache.MeshCardsMinSize=4
 r.Material.EnergyConservation=True
 r.Mobile.FloatPrecisionMode=2
+r.MotionBlurQuality=0
 r.OIT.SortedPixels=True
+r.PathTracing.Denoiser.Prepass.RankedLuminanceVariance=0
 r.PathTracing.VisibleLights=1
 r.RayTracedTranslucency.Allowed=1
 r.RayTracedTranslucency.ForceOpaque=0
 r.RayTracedTranslucency.MaxPrimaryHitEvents=3
 r.RayTracedTranslucency.UseRayTracedRefraction=1
 r.RayTracing.Culling=0
-r.RayTracing.Nanite.Mode=1
+r.RayTracing.Nanite.Mode=0
 r.RayTracing.NormalBias=0.18
-r.RayTracing.ResidentGeometryMemoryPoolSizeInMB=700
 r.RayTracing.Shadows.Translucency=1
 r.RayTracing.Shadows=True
 r.RayTracing=True
 r.ReflectionCaptureResolution=1024
 r.ReflectionMethod=1
+r.ScreenPercentage.Default.Desktop.Mode=0
 r.Shadow.Virtual.Enable=1
+r.Shadow.Virtual.TranslucentQuality=1
 r.SkinCache.CompileShaders=True
 r.SkinCache.SceneMemoryLimitInMB=126.000000
 r.SkyAtmosphere.MultiScatteringLUT.HighQuality=1
-r.Substrate.BytesPerPixel=240
-r.Substrate.Debug.AdvancedVisualizationShaders=True
-r.Substrate.EnableLayerSupport=True
+r.Substrate.BytesPerPixel=160
 r.Substrate.OpaqueMaterialRoughRefraction=True
-r.Substrate.ProjectClosuresPerPixel=5
+r.Substrate.ProjectClosuresPerPixel=4
 r.Substrate.ProjectClosuresPerPixelOverride=1
 r.Substrate.ProjectGBufferFormat=1
 r.Substrate=True
+r.TSR.History.SampleCount=32
+r.TSR.ThinGeometryDetection.Coverage.ShadingRange=2
 r.Translucency.StandardSeparated=1
 r.Translucent.UsesRectLights=True
 r.VirtualTextures=True
 
+
 [/Script/InterchangeEngine.InterchangeProjectSettings]
 ContentImportSettings=(DefaultPipelineStackOverride=((Materials, "Materials & Textures"),(Textures, "Textures"),(Sounds, "Audio")),ShowImportDialogOverride=(),PipelineStacks=(("Static Meshes", (Pipelines=("/Game/InterchangePipeline/Pipeline_StaticMeshes.Pipeline_StaticMeshes"),PerTranslatorPipelines=((Translator="/Script/InterchangeImport.InterchangeMaterialXTranslator",Pipelines=("/Interchange/Pipelines/DefaultAssetsPipeline.DefaultAssetsPipeline","/Interchange/Pipelines/DefaultMaterialXPipeline.DefaultMaterialXPipeline")),(Translator="/Script/InterchangeImport.InterchangeGLTFTranslator",Pipelines=("/Interchange/Pipelines/DefaultGLTFAssetsPipeline.DefaultGLTFAssetsPipeline","/Interchange/Pipelines/DefaultGLTFPipeline.DefaultGLTFPipeline"))))),("Skeletal Meshes", (Pipelines=("/Game/InterchangePipeline/Pipeline_SkeletalMeshes.Pipeline_SkeletalMeshes"),PerTranslatorPipelines=((Translator="/Script/InterchangeImport.InterchangeMaterialXTranslator",Pipelines=("/Interchange/Pipelines/DefaultAssetsPipeline.DefaultAssetsPipeline","/Interchange/Pipelines/DefaultMaterialXPipeline.DefaultMaterialXPipeline")),(Translator="/Script/InterchangeImport.InterchangeGLTFTranslator",Pipelines=("/Interchange/Pipelines/DefaultGLTFAssetsPipeline.DefaultGLTFAssetsPipeline","/Interchange/Pipelines/DefaultGLTFPipeline.DefaultGLTFPipeline"))))),("Audio", (Pipelines=("/Interchange/Pipelines/DefaultAudioPipeline.DefaultAudioPipeline"))),("Materials & Textures", (Pipelines=("/Interchange/Pipelines/DefaultMaterialPipeline.DefaultMaterialPipeline"),PerTranslatorPipelines=((Translator="/Script/InterchangeImport.InterchangeMaterialXTranslator",Pipelines=("/Interchange/Pipelines/DefaultMaterialPipeline.DefaultMaterialPipeline","/Interchange/Pipelines/DefaultMaterialXPipeline.DefaultMaterialXPipeline"))))),("Textures", (Pipelines=("/Interchange/Pipelines/DefaultTexturePipeline.DefaultTexturePipeline"))),("HDR Cubemap", (Pipelines=("/Game/InterchangePipeline/Pipeline_HDR_Cubemap.Pipeline_HDR_Cubemap")))),DefaultPipelineStack="Static Meshes",ImportDialogClass="/Script/InterchangeEditorPipelines.InterchangePipelineConfigurationGeneric",bShowImportDialog=True,bShowReimportDialog=False)
 SceneImportSettings=(PerTranslatorDialogOverride=,PipelineStacks=(("Scene", (Pipelines=("/Interchange/Pipelines/DefaultSceneAssetsPipeline.DefaultSceneAssetsPipeline","/Interchange/Pipelines/DefaultSceneLevelPipeline.DefaultSceneLevelPipeline"),PerTranslatorPipelines=((Translator="/Script/InterchangeImport.InterchangeMaterialXTranslator",Pipelines=("/Interchange/Pipelines/DefaultSceneAssetsPipeline.DefaultSceneAssetsPipeline","/Interchange/Pipelines/DefaultSceneLevelPipeline.DefaultSceneLevelPipeline","/Interchange/Pipelines/DefaultMaterialXPipeline.DefaultMaterialXPipeline")),(Translator="/Script/InterchangeImport.InterchangeGLTFTranslator",Pipelines=("/Interchange/Pipelines/DefaultGLTFSceneAssetsPipeline.DefaultGLTFSceneAssetsPipeline","/Interchange/Pipelines/DefaultSceneLevelPipeline.DefaultSceneLevelPipeline","/Interchange/Pipelines/DefaultGLTFPipeline.DefaultGLTFPipeline")))))),DefaultPipelineStack="Scene",ImportDialogClass="/Script/InterchangeEditorPipelines.InterchangePipelineConfigurationGeneric",bShowImportDialog=True,bShowReimportDialog=False)
diff --git a/Templates/TP_PhotoStudioBP/Config/DefaultEngine.ini b/Templates/TP_PhotoStudioBP/Config/DefaultEngine.ini
--- a/Templates/TP_PhotoStudioBP/Config/DefaultEngine.ini
+++ b/Templates/TP_PhotoStudioBP/Config/DefaultEngine.ini
@@ -16,7 +16,7 @@ DefaultGraphicsRHI=DefaultGraphicsRHI_DX12
 [/Script/Engine.RendererSettings]
 r.AllowStaticLighting=0
 r.AntiAliasingMethod=4
-r.ClearCoatNormal=True
+r.BloomQuality=2
 r.DefaultFeature.AutoExposure.ExtendDefaultLuminanceRange=True
 r.DefaultFeature.AutoExposure=False
 r.DefaultFeature.LocalExposure.HighlightContrastScale=0.8
@@ -28,43 +28,57 @@ r.GenerateMeshDistanceFields=True
 r.Lumen.HardwareRayTracing.HitLighting.ForceOpaque=0
 r.Lumen.HardwareRayTracing.LightingMode=1
 r.Lumen.LightingDataFormat=1
+r.Lumen.Reflections.BilateralFilter.KernelRadius=32
 r.Lumen.Reflections.DenoiserTonemapRange =12
 r.Lumen.Reflections.HardwareRayTracing.Translucent.MaxRefractionBounces=2
+r.Lumen.ScreenProbeGather.DownsampleFactor=4
+r.Lumen.ScreenProbeGather.ShortRangeAO.HorizonSearch.ForegroundSampleRejectDistanceFraction=0.025
+r.Lumen.ScreenProbeGather.Temporal.DistanceThreshold=0.001
+r.Lumen.ScreenProbeGather.Temporal.NormalThreshold=5
+r.Lumen.ScreenProbeGather.Temporal.RejectBasedOnNormal=1
+r.Lumen.ScreenProbeGather.TemporalFilterProbes=1
 r.Lumen.TraceMeshSDFs=0
 r.Lumen.TranslucencyReflections.FrontLayer.EnableForProject=True
+r.LumenScene.SurfaceCache.CardMinResolution=2
+r.LumenScene.SurfaceCache.MeshCardsMinSize=4
 r.Material.EnergyConservation=True
 r.Mobile.FloatPrecisionMode=2
+r.MotionBlurQuality=0
 r.OIT.SortedPixels=True
+r.PathTracing.Denoiser.Prepass.RankedLuminanceVariance=0
 r.PathTracing.VisibleLights=1
 r.RayTracedTranslucency.Allowed=1
 r.RayTracedTranslucency.ForceOpaque=0
 r.RayTracedTranslucency.MaxPrimaryHitEvents=3
 r.RayTracedTranslucency.UseRayTracedRefraction=1
 r.RayTracing.Culling=0
-r.RayTracing.Nanite.Mode=1
+r.RayTracing.Nanite.Mode=0
 r.RayTracing.NormalBias=0.18
-r.RayTracing.ResidentGeometryMemoryPoolSizeInMB=700
 r.RayTracing.Shadows.Translucency=1
 r.RayTracing.Shadows=True
 r.RayTracing=True
 r.ReflectionCaptureResolution=1024
 r.ReflectionMethod=1
+r.ScreenPercentage.Default.Desktop.Mode=0
 r.Shadow.Virtual.Enable=1
+r.Shadow.Virtual.TranslucentQuality=1
 r.SkinCache.CompileShaders=True
 r.SkinCache.SceneMemoryLimitInMB=126.000000
 r.SkyAtmosphere.MultiScatteringLUT.HighQuality=1
-r.Substrate.BytesPerPixel=240
-r.Substrate.Debug.AdvancedVisualizationShaders=True
-r.Substrate.EnableLayerSupport=True
+r.Substrate.BytesPerPixel=160
 r.Substrate.OpaqueMaterialRoughRefraction=True
-r.Substrate.ProjectClosuresPerPixel=5
+r.Substrate.ProjectClosuresPerPixel=4
 r.Substrate.ProjectClosuresPerPixelOverride=1
 r.Substrate.ProjectGBufferFormat=1
 r.Substrate=True
+r.TSR.History.SampleCount=32
+r.TSR.ThinGeometryDetection.Coverage.ShadingRange=2
 r.Translucency.StandardSeparated=1
 r.Translucent.UsesRectLights=True
 r.VirtualTextures=True
 
+
+
 [/Script/Engine.VirtualTexturePoolConfig]
 DefaultSizeInMegabyte=64
 bPoolAutoGrowInEditor=False
